Added ChaosInsights module for Chaos related extensions to insights and implemented a physics scene lock profiler.

Often a multithreaded process will inadvertently open a project up to contention on the main physics scene lock. Queries require a read lock on this data and any update to the location of an object requires a write lock. If a worker begins running a large number of queries or holding a lock over a period where the game thread moves a component or attempts to sync from the physics simulation it can cause stalls on the game thread.

This is often difficult to spot so enabling the "ChaosLocks" channel in an insights capture will enabled the physics scene lock profiler which visualises all of the threads attempting to hold a lock on the physics scene to help detect these conditions. Is shows reads and writes separately and clearly shows where threads are waiting on locks before proceeding. It also reports how many recursive locks were taken during each lock region for reasoning about operation density.

diff --git a/Engine/Plugins/ChaosInsights/ChaosInsights.uplugin b/Engine/Plugins/ChaosInsights/ChaosInsights.uplugin
new file mode 100644
--- /dev/null
+++ b/Engine/Plugins/ChaosInsights/ChaosInsights.uplugin
@@ -0,0 +1,33 @@
+{
+	"FileVersion": 3,
+	"Version": 1,
+	"VersionName": "1.0",
+	"FriendlyName": "Chaos Insights",
+	"Description": "Plugin to gather insights into Chaos",
+	"Category": "Insights",
+	"CreatedBy": "Epic Games, Inc.",
+	"CreatedByURL": "https://epicgames.com",
+	"EnabledByDefault": true,
+	"CanContainContent": false,
+	"IsBetaVersion": true,
+	"Installed": false,
+	"SupportedPrograms": [ "UnrealInsights" ],
+	"Modules": [
+	{
+		"Name": "ChaosInsightsAnalysis",
+		"Type": "EditorAndProgram",
+		"LoadingPhase": "Default",
+		"ProgramAllowList": [
+		"UnrealInsights"
+			]
+	},
+	{
+		"Name": "ChaosInsightsUI",
+		"Type": "EditorAndProgram",
+		"LoadingPhase": "Default",
+		"ProgramAllowList": [
+		"UnrealInsights"
+			]
+	}
+	]
+}
\ No newline at end of file
diff --git a/Engine/Plugins/ChaosInsights/Source/ChaosInsightsAnalysis/ChaosInsightsAnalysis.build.cs b/Engine/Plugins/ChaosInsights/Source/ChaosInsightsAnalysis/ChaosInsightsAnalysis.build.cs
new file mode 100644
--- /dev/null
+++ b/Engine/Plugins/ChaosInsights/Source/ChaosInsightsAnalysis/ChaosInsightsAnalysis.build.cs
@@ -0,0 +1,22 @@
+// Copyright Epic Games, Inc. All Rights Reserved.
+
+namespace UnrealBuildTool.Rules
+{
+	public class ChaosInsightsAnalysis : ModuleRules
+	{
+		public ChaosInsightsAnalysis(ReadOnlyTargetRules Target) : base(Target)
+		{
+			PublicDependencyModuleNames.AddRange(new string[] { "Core" });
+			
+			PrivateDependencyModuleNames.AddRange(
+				new string[]
+				{
+					"CoreUObject",
+					"TraceLog",
+					"TraceAnalysis",
+					"TraceServices",
+				});
+		}
+	}
+}
+
diff --git a/Engine/Plugins/ChaosInsights/Source/ChaosInsightsAnalysis/Private/Analyzers/ChaosInsightsTraceAnalysis.cpp b/Engine/Plugins/ChaosInsights/Source/ChaosInsightsAnalysis/Private/Analyzers/ChaosInsightsTraceAnalysis.cpp
new file mode 100644
--- /dev/null
+++ b/Engine/Plugins/ChaosInsights/Source/ChaosInsightsAnalysis/Private/Analyzers/ChaosInsightsTraceAnalysis.cpp
@@ -0,0 +1,78 @@
+// Copyright Epic Games, Inc. All Rights Reserved.
+
+#include "ChaosInsightsTraceAnalysis.h"
+
+#include "Common/ProviderLock.h"
+#include "TraceServices/Model/AnalysisSession.h"
+#include "Model/LockRegions.h"
+
+namespace ChaosInsightsAnalysis
+{
+
+	FLockRegionsTraceAnalyzer::FLockRegionsTraceAnalyzer(TraceServices::IAnalysisSession& InSession, FLockRegionProvider& InRegionProvider)
+		: Session(InSession)
+		, RegionProvider(InRegionProvider)
+	{
+	}
+
+	void FLockRegionsTraceAnalyzer::OnAnalysisBegin(const FOnAnalysisContext& Context)
+	{
+		FInterfaceBuilder& Builder = Context.InterfaceBuilder;
+
+		Builder.RouteEvent(static_cast<uint16>(ELockRegionEventId::LockBegin), "Chaos", "LockAcquireBegin");
+		Builder.RouteEvent(static_cast<uint16>(ELockRegionEventId::LockAcquire), "Chaos", "LockAcquired");
+		Builder.RouteEvent(static_cast<uint16>(ELockRegionEventId::LockEnd), "Chaos", "LockAcquireEnd");
+	}
+
+	void FLockRegionsTraceAnalyzer::OnAnalysisEnd()
+	{
+		TraceServices::FProviderEditScopeLock RegionProviderScopedLock(static_cast<TraceServices::IEditableProvider&>(RegionProvider));
+		RegionProvider.OnAnalysisSessionEnded();
+	}
+
+	bool FLockRegionsTraceAnalyzer::OnEvent(uint16 RouteId, EStyle Style, const FOnEventContext& Context)
+	{
+		const FEventData& EventData = Context.EventData;
+		const ELockRegionEventId Id = static_cast<ELockRegionEventId>(RouteId);
+
+		switch(Id)
+		{
+		case ELockRegionEventId::LockBegin:
+		{
+			uint64 Cycle = EventData.GetValue<uint64>("Cycle");
+			bool bWriteLock = EventData.GetValue<bool>("bIsWrite");
+			
+			TraceServices::FProviderEditScopeLock ProviderLock(RegionProvider);
+			RegionProvider.AppendRegionBegin(Context.EventTime.AsSeconds(Cycle), Context.ThreadInfo.GetId(), bWriteLock);
+
+			break;
+		}
+		case ELockRegionEventId::LockAcquire:
+		{
+			uint64 Cycle = EventData.GetValue<uint64>("Cycle");
+			
+			TraceServices::FProviderEditScopeLock ProviderLock(RegionProvider);
+			RegionProvider.AppendRegionAcquired(Context.EventTime.AsSeconds(Cycle), Context.ThreadInfo.GetId());
+
+			break;
+		}
+		case ELockRegionEventId::LockEnd:
+		{
+			uint64 Cycle = EventData.GetValue<uint64>("Cycle");
+			
+			TraceServices::FProviderEditScopeLock ProviderLock(RegionProvider);
+			RegionProvider.AppendRegionEnd(Context.EventTime.AsSeconds(Cycle), Context.ThreadInfo.GetId());
+
+			break;
+		}
+		default:
+		{
+			// Unrequested event ID
+			check(false);
+		}
+		}
+
+		return true;
+	}
+
+}
diff --git a/Engine/Plugins/ChaosInsights/Source/ChaosInsightsAnalysis/Private/Analyzers/ChaosInsightsTraceAnalysis.h b/Engine/Plugins/ChaosInsights/Source/ChaosInsightsAnalysis/Private/Analyzers/ChaosInsightsTraceAnalysis.h
new file mode 100644
--- /dev/null
+++ b/Engine/Plugins/ChaosInsights/Source/ChaosInsightsAnalysis/Private/Analyzers/ChaosInsightsTraceAnalysis.h
@@ -0,0 +1,39 @@
+// Copyright Epic Games, Inc. All Rights Reserved.
+
+#pragma once
+
+#include "Trace/Analyzer.h"
+#include "Templates/SharedPointer.h"
+
+namespace TraceServices
+{
+	class IAnalysisSession;
+}
+
+namespace ChaosInsightsAnalysis
+{
+	class FLockRegionProvider;
+
+	class FLockRegionsTraceAnalyzer : public UE::Trace::IAnalyzer
+	{
+	public:
+		FLockRegionsTraceAnalyzer(TraceServices::IAnalysisSession& Session, FLockRegionProvider& RegionProvider);
+
+		virtual void OnAnalysisBegin(const FOnAnalysisContext& Context) override;
+		virtual void OnAnalysisEnd() override;
+		virtual bool OnEvent(uint16 RouteId, EStyle Style, const FOnEventContext& Context) override;
+
+	private:
+		enum class ELockRegionEventId : uint16
+		{
+			LockBegin,
+			LockAcquire,
+			LockEnd
+		};
+
+		TraceServices::IAnalysisSession& Session;
+		FLockRegionProvider& RegionProvider;
+	};
+
+
+}
diff --git a/Engine/Plugins/ChaosInsights/Source/ChaosInsightsAnalysis/Private/ChaosInsightsAnalysisModule.cpp b/Engine/Plugins/ChaosInsights/Source/ChaosInsightsAnalysis/Private/ChaosInsightsAnalysisModule.cpp
new file mode 100644
--- /dev/null
+++ b/Engine/Plugins/ChaosInsights/Source/ChaosInsightsAnalysis/Private/ChaosInsightsAnalysisModule.cpp
@@ -0,0 +1,37 @@
+// Copyright Epic Games, Inc. All Rights Reserved.
+
+#include "ChaosInsightsAnalysisModule.h"
+#include "Modules/ModuleManager.h"
+#include "Features/IModularFeatures.h"
+#include "Model/LockRegions.h"
+#include "Analyzers/ChaosInsightsTraceAnalysis.h"
+
+namespace ChaosInsightsAnalysis
+{
+	void FChaosInsightsAnalysisModule::OnAnalysisBegin(TraceServices::IAnalysisSession& InSession)
+	{
+		TSharedPtr<FLockRegionProvider> LockRegionsProvider = MakeShared<FLockRegionProvider>(InSession);
+		InSession.AddProvider(GetLockRegionProviderName(), LockRegionsProvider);
+
+		InSession.AddAnalyzer(new FLockRegionsTraceAnalyzer(InSession, *LockRegionsProvider));
+	}
+
+	void FChaosInsightsAnalysisModule::GetModuleInfo(TraceServices::FModuleInfo& OutModuleInfo)
+	{
+		OutModuleInfo.Name = TEXT("ChaosInsightsAnalysis");
+		OutModuleInfo.DisplayName = TEXT("Chaos Insights");
+	}
+
+	void FChaosInsightsAnalysisModule::StartupModule()
+	{
+		IModularFeatures::Get().RegisterModularFeature(TraceServices::ModuleFeatureName, this);
+	}
+
+	void FChaosInsightsAnalysisModule::ShutdownModule()
+	{
+		IModularFeatures::Get().UnregisterModularFeature(TraceServices::ModuleFeatureName, this);
+	}
+
+}
+
+IMPLEMENT_MODULE(ChaosInsightsAnalysis::FChaosInsightsAnalysisModule, ChaosInsightsAnalysis);
diff --git a/Engine/Plugins/ChaosInsights/Source/ChaosInsightsAnalysis/Private/ChaosInsightsAnalysisModule.h b/Engine/Plugins/ChaosInsights/Source/ChaosInsightsAnalysis/Private/ChaosInsightsAnalysisModule.h
new file mode 100644
--- /dev/null
+++ b/Engine/Plugins/ChaosInsights/Source/ChaosInsightsAnalysis/Private/ChaosInsightsAnalysisModule.h
@@ -0,0 +1,25 @@
+// Copyright Epic Games, Inc. All Rights Reserved.
+
+#pragma once
+
+#include "Modules/ModuleInterface.h"
+#include "TraceServices/ModuleService.h"
+
+class FTabManager;
+
+namespace ChaosInsightsAnalysis
+{
+	class FChaosInsightsAnalysisModule : public IModuleInterface, TraceServices::IModule
+	{
+	public:
+		// IModuleInterface interface
+		virtual void StartupModule() override;
+		virtual void ShutdownModule() override;
+
+	public:
+		// TraceServices::IModule interface
+		virtual void GetModuleInfo(TraceServices::FModuleInfo& OutModuleInfo) override;
+		virtual void OnAnalysisBegin(TraceServices::IAnalysisSession& Session) override;
+	};
+
+}
diff --git a/Engine/Plugins/ChaosInsights/Source/ChaosInsightsAnalysis/Private/Model/LockRegions.cpp b/Engine/Plugins/ChaosInsights/Source/ChaosInsightsAnalysis/Private/Model/LockRegions.cpp
new file mode 100644
--- /dev/null
+++ b/Engine/Plugins/ChaosInsights/Source/ChaosInsightsAnalysis/Private/Model/LockRegions.cpp
@@ -0,0 +1,299 @@
+// Copyright Epic Games, Inc. All Rights Reserved.
+
+#include "ChaosInsightsAnalysis/Model/LockRegions.h"
+#include "Model/LockRegions.h"
+
+#include "Algo/ForEach.h"
+#include "Internationalization/Internationalization.h"
+#include "TraceServices/Model/Threads.h"
+
+#define LOCTEXT_NAMESPACE "LockRegionProvider"
+
+DEFINE_LOG_CATEGORY(LogChaosInsights)
+
+namespace ChaosInsightsAnalysis
+{
+
+	thread_local TraceServices::FProviderLock::FThreadLocalState GLockRegionsProviderLockState;
+
+	FLockRegionProvider::FLockRegionProvider(TraceServices::IAnalysisSession& InSession)
+		: Session(InSession)
+	{
+	}
+
+	FLockRegionProvider::~FLockRegionProvider()
+	{
+
+	}
+
+	void FLockRegionProvider::BeginRead() const
+	{
+		Lock.BeginRead(GLockRegionsProviderLockState);
+	}
+
+	void FLockRegionProvider::EndRead() const
+	{
+		Lock.EndRead(GLockRegionsProviderLockState);
+	}
+
+	void FLockRegionProvider::ReadAccessCheck() const
+	{
+		Lock.ReadAccessCheck(GLockRegionsProviderLockState);
+	}
+
+	uint64 FLockRegionProvider::GetRegionCount() const
+	{
+		ReadAccessCheck();
+
+		uint64 RegionCount = 0;
+		for(const FLockRegionLane& Lane : Lanes)
+		{
+			RegionCount += Lane.Num();
+		}
+		return RegionCount;
+	}
+
+	int32 FLockRegionProvider::GetLaneCount() const
+	{
+		ReadAccessCheck(); return Lanes.Num();
+	}
+
+	const FLockRegionLane* FLockRegionProvider::GetLane(int32 index) const
+	{
+		ReadAccessCheck();
+
+		if(index < Lanes.Num())
+		{
+			return &(Lanes[index]);
+		}
+		return nullptr;
+	}
+
+	void FLockRegionProvider::AppendRegionBegin(double Time, uint64 ThreadId, bool bWrite)
+	{
+		EditAccessCheck();
+		
+		FLockRegion** FoundRegion = OpenRegionsByThread.Find(ThreadId);
+
+		int32 InitialLockDepth = 0;
+
+		if(FoundRegion)
+		{
+			(*FoundRegion)->LockCount++;
+			(*FoundRegion)->LockDepth++;
+		}
+		else
+		{
+			int32 Depth = CalculateRegionDepth(Time);
+
+			while(Depth >= Lanes.Num())
+			{
+				Lanes.Emplace(Session.GetLinearAllocator());
+			}
+
+			TraceServices::FAnalysisSessionReadScope SessionReadScope(Session);
+			const TraceServices::IThreadProvider& ThreadProvider = TraceServices::ReadThreadProvider(Session);
+
+			Lanes[Depth].Regions.EmplaceBack(FLockRegion
+				{
+					.BeginTime = Time,
+					.Text = Session.StoreString(ThreadProvider.GetThreadName(ThreadId)),
+					.Thread = ThreadId,
+					.Depth = Depth,
+					.LockCount = 1,
+					.LockDepth = 1,
+					.bIsWrite = bWrite
+				});
+
+			OpenRegionsByThread.Add(ThreadId, &Lanes[Depth].Regions.Last());
+		}
+
+		UpdateSession(Time);
+	}
+
+	void FLockRegionProvider::AppendRegionAcquired(double Time, uint64 ThreadId)
+	{
+		EditAccessCheck();
+
+		FLockRegion* Region = OpenRegionsByThread.FindChecked(ThreadId);
+
+		if(Time < Region->AcquireTime)
+		{
+			Region->AcquireTime = Time;
+		}
+
+		UpdateSession(Time);
+	}
+
+	void FLockRegionProvider::AppendRegionEnd(double Time, uint64 ThreadId)
+	{
+		EditAccessCheck();
+
+		TMap<uint64_t, FLockRegion*>::TKeyIterator RegionIt = OpenRegionsByThread.CreateKeyIterator(ThreadId);
+		FLockRegion& Region = *RegionIt->Value;
+
+		if(--Region.LockDepth == 0)
+		{
+			RegionIt->Value->EndTime = Time;
+
+			RegionIt.RemoveCurrent();
+		}
+
+		UpdateSession(Time);
+	}
+
+	void FLockRegionProvider::OnAnalysisSessionEnded()
+	{
+		EditAccessCheck();
+
+		if(OpenRegionsByThread.Num() > 0)
+		{
+			UE_LOG(LogChaosInsights, Warning, TEXT("A physics lock event was never closed."));
+		}
+	}
+
+	int32 FLockRegionProvider::CalculateRegionDepth(double NewBeginTime) const
+	{
+		constexpr int32 DepthLimit = 100;
+
+		int32 NewDepth = 0;
+
+		// Find first free lane/depth
+		while(NewDepth < DepthLimit)
+		{
+			if(!Lanes.IsValidIndex(NewDepth))
+			{
+				break;
+			}
+
+			const FLockRegion& LastRegion = Lanes[NewDepth].Regions.Last();
+			if(LastRegion.EndTime <= NewBeginTime)
+			{
+				break;
+			}
+			NewDepth++;
+		}
+
+		ensureMsgf(NewDepth < DepthLimit, TEXT("Regions are nested too deep."));
+
+		return NewDepth;
+	}
+
+	void FLockRegionProvider::UpdateSession(double InTime)
+	{
+		TraceServices::FAnalysisSessionEditScope EditScope(Session);
+		Session.UpdateDurationSeconds(InTime);
+	}
+
+	void FLockRegionProvider::ForEachLane(TFunctionRef<void(const FLockRegionLane&, int32)> Callback) const
+	{
+		ReadAccessCheck();
+
+		for(int32 LaneIndex = 0; LaneIndex < Lanes.Num(); ++LaneIndex)
+		{
+			Callback(Lanes[LaneIndex], LaneIndex);
+		}
+	}
+
+	void FLockRegionProvider::BeginEdit() const
+	{
+		Lock.BeginWrite(GLockRegionsProviderLockState);
+	}
+
+	void FLockRegionProvider::EndEdit() const
+	{
+		Lock.EndWrite(GLockRegionsProviderLockState);
+	}
+
+	void FLockRegionProvider::EditAccessCheck() const
+	{
+		Lock.WriteAccessCheck(GLockRegionsProviderLockState);
+	}
+
+	bool FLockRegionProvider::ForEachRegionInRange(double IntervalStart, double IntervalEnd, TFunctionRef<bool(const FLockRegion&)> Callback) const
+	{
+		ReadAccessCheck();
+
+		if(IntervalStart > IntervalEnd)
+		{
+			return false;
+		}
+
+		for(const FLockRegionLane& Lane : Lanes)
+		{
+			if(!Lane.ForEachRegionInRange(IntervalStart, IntervalEnd, Callback))
+			{
+				return false;
+			}
+		}
+
+		return true;
+	}
+
+	FLockRegionLane::FLockRegionLane(TraceServices::ILinearAllocator& InAllocator) 
+		: Regions(InAllocator, 512)
+	{
+
+	}
+
+	int32 FLockRegionLane::Num() const
+	{
+		return static_cast<int32>(Regions.Num());
+	}
+
+	bool FLockRegionLane::ForEachRegionInRange(double IntervalStart, double IntervalEnd, TFunctionRef<bool(const FLockRegion&)> Callback) const
+	{
+		struct FTimeProjections
+		{
+			static double Begin(const FLockRegion& InRegion)
+			{
+				return InRegion.BeginTime;
+			}
+
+			static double End(const FLockRegion& InRegion)
+			{
+				return InRegion.EndTime;
+			}
+		};
+
+		const FInt32Interval OverlapRange = GetElementRangeOverlappingGivenRange<FLockRegion>(Regions, IntervalStart, IntervalEnd, &FTimeProjections::Begin, &FTimeProjections::End);
+
+		if(OverlapRange.Min == -1)
+		{
+			return true;
+		}
+
+		for(int32 Index = OverlapRange.Min; Index <= OverlapRange.Max; ++Index)
+		{
+			if(!Callback(Regions[Index]))
+			{
+				return false;
+			}
+		}
+
+		return true;
+	}
+
+	FName GetLockRegionProviderName()
+	{
+		static const FName Name("LockRegionProvider");
+		return Name;
+	}
+
+	ILockRegionProvider::~ILockRegionProvider() = default;
+
+	const ILockRegionProvider& ReadRegionProvider(const TraceServices::IAnalysisSession& Session)
+	{
+		return *Session.ReadProvider<ILockRegionProvider>(GetLockRegionProviderName());
+	}
+
+	IEditableLockRegionProvider::~IEditableLockRegionProvider() = default;
+
+	IEditableLockRegionProvider& EditRegionProvider(TraceServices::IAnalysisSession& Session)
+	{
+		return *Session.EditProvider<IEditableLockRegionProvider>(GetLockRegionProviderName());
+	}
+
+}
+
+#undef LOCTEXT_NAMESPACE
diff --git a/Engine/Plugins/ChaosInsights/Source/ChaosInsightsAnalysis/Private/Model/LockRegions.h b/Engine/Plugins/ChaosInsights/Source/ChaosInsightsAnalysis/Private/Model/LockRegions.h
new file mode 100644
--- /dev/null
+++ b/Engine/Plugins/ChaosInsights/Source/ChaosInsightsAnalysis/Private/Model/LockRegions.h
@@ -0,0 +1,68 @@
+// Copyright Epic Games, Inc. All Rights Reserved.
+
+#pragma once
+
+#include "Common/ProviderLock.h"
+#include "Templates/SharedPointer.h"
+#include "ChaosInsightsAnalysis/Model/LockRegions.h"
+
+DECLARE_LOG_CATEGORY_EXTERN(LogChaosInsights, Display, All)
+
+namespace TraceServices
+{
+	class FAnalysisSessionLock;
+	class FStringStore;
+}
+
+namespace ChaosInsightsAnalysis
+{
+	extern thread_local TraceServices::FProviderLock::FThreadLocalState GLockRegionsProviderLockState;
+
+	class FLockRegionProvider : public ILockRegionProvider, public IEditableLockRegionProvider
+	{
+	public:
+		explicit FLockRegionProvider(TraceServices::IAnalysisSession& Session);
+		virtual ~FLockRegionProvider() override;
+
+		//////////////////////////////////////////////////
+		// Read operations
+
+		virtual void BeginRead() const override;
+		virtual void EndRead() const override;
+		virtual void ReadAccessCheck() const override;
+
+		virtual uint64 GetRegionCount() const override;
+		virtual int32 GetLaneCount()  const override;
+
+		virtual const FLockRegionLane* GetLane(int32 Index) const override;
+
+		virtual bool ForEachRegionInRange(double IntervalStart, double IntervalEnd, TFunctionRef<bool(const FLockRegion&)> Callback) const override;
+		virtual void ForEachLane(TFunctionRef<void(const FLockRegionLane&, const int32)> Callback) const override;
+
+		//////////////////////////////////////////////////
+		// Edit operations
+
+		virtual void BeginEdit() const override;
+		virtual void EndEdit() const override;
+		virtual void EditAccessCheck() const override;
+
+		virtual void AppendRegionBegin(double Time, uint64 ThreadId, bool bWrite) override;
+		virtual void AppendRegionAcquired(double Time, uint64 ThreadId) override;
+		virtual void AppendRegionEnd(double Time, uint64 ThreadId) override;
+		
+		virtual void OnAnalysisSessionEnded() override;
+
+		//////////////////////////////////////////////////
+
+	private:
+		int32 CalculateRegionDepth(double NewBeginTime) const;
+		void UpdateSession(double InTime);
+		
+		TraceServices::IAnalysisSession& Session;
+		mutable TraceServices::FProviderLock Lock;
+
+		TMap<uint64_t, FLockRegion*> OpenRegionsByThread;
+		TArray<FLockRegionLane> Lanes;
+	};
+
+}
diff --git a/Engine/Plugins/ChaosInsights/Source/ChaosInsightsAnalysis/Public/ChaosInsightsAnalysis/Model/LockRegions.h b/Engine/Plugins/ChaosInsights/Source/ChaosInsightsAnalysis/Public/ChaosInsightsAnalysis/Model/LockRegions.h
new file mode 100644
--- /dev/null
+++ b/Engine/Plugins/ChaosInsights/Source/ChaosInsightsAnalysis/Public/ChaosInsightsAnalysis/Model/LockRegions.h
@@ -0,0 +1,90 @@
+// Copyright Epic Games, Inc. All Rights Reserved.
+
+#pragma once
+
+#include "Common/PagedArray.h" // TraceServices
+#include "HAL/Platform.h"
+#include "Templates/Function.h"
+#include "TraceServices/Model/AnalysisSession.h"
+#include "UObject/NameTypes.h"
+
+namespace ChaosInsightsAnalysis
+{
+
+	struct CHAOSINSIGHTSANALYSIS_API FLockRegion
+	{
+		// Time that the caller attempted to take the lock
+		double BeginTime = std::numeric_limits<double>::infinity();
+
+		// Time that the lock was actually acquired (including waiting on the lock if it was already taken)
+		double AcquireTime = std::numeric_limits<double>::infinity();
+
+		// Time that the caller relinquished the lock
+		double EndTime = std::numeric_limits<double>::infinity();
+
+		// Name of the thread that took the lock
+		const TCHAR* Text = nullptr;
+
+		// The local thread Id for the thread that took the lock
+		uint64 Thread;
+
+		// UI depth
+		int32 Depth = -1;
+
+		// Number of times a lock was taken during the region
+		int32 LockCount = 0;
+
+		// Tracking for the lock depth to combine all recursive locks into one region
+		int32 LockDepth = 0;
+
+		// Whether this is a write-lock (as opposed to a read-lock)
+		bool bIsWrite = false;
+	};
+
+	class CHAOSINSIGHTSANALYSIS_API FLockRegionLane
+	{
+		friend class FLockRegionProvider;
+
+	public:
+		FLockRegionLane(TraceServices::ILinearAllocator& InAllocator);
+
+		int32 Num() const;
+
+		bool ForEachRegionInRange(double IntervalStart, double IntervalEnd, TFunctionRef<bool(const FLockRegion&)> Callback) const;
+
+	private:
+		TraceServices::TPagedArray<FLockRegion> Regions;
+	};
+
+	class CHAOSINSIGHTSANALYSIS_API ILockRegionProvider : public TraceServices::IProvider
+	{
+	public:
+		virtual ~ILockRegionProvider() override;
+
+		virtual uint64 GetRegionCount() const = 0;
+		virtual int32 GetLaneCount() const = 0;
+		virtual const FLockRegionLane* GetLane(int32 Index) const = 0;
+
+		virtual bool ForEachRegionInRange(double IntervalStart, double IntervalEnd, TFunctionRef<bool(const FLockRegion&)> Callback) const = 0;
+		virtual void ForEachLane(TFunctionRef<void(const FLockRegionLane&, const int32)> Callback) const = 0;
+	};
+
+	class IEditableLockRegionProvider : public TraceServices::IEditableProvider
+	{
+	public:
+		virtual ~IEditableLockRegionProvider() override;
+
+		virtual void AppendRegionBegin(double Time, uint64 ThreadId, bool bIsWrite) = 0;
+		virtual void AppendRegionAcquired(double Time, uint64 ThreadId) = 0;
+		virtual void AppendRegionEnd(double Time, uint64 ThreadId) = 0;
+
+		virtual void OnAnalysisSessionEnded() = 0;
+	};
+
+	////////////////////////////////////////////////////////////////////////////////////////////////////
+
+	CHAOSINSIGHTSANALYSIS_API FName GetLockRegionProviderName();
+	CHAOSINSIGHTSANALYSIS_API const ILockRegionProvider& ReadRegionProvider(const TraceServices::IAnalysisSession& Session);
+	CHAOSINSIGHTSANALYSIS_API IEditableLockRegionProvider& EditRegionProvider(TraceServices::IAnalysisSession& Session);
+
+} // namespace TimingRegionsAnalysis
diff --git a/Engine/Plugins/ChaosInsights/Source/ChaosInsightsUI/ChaosInsightsUI.build.cs b/Engine/Plugins/ChaosInsights/Source/ChaosInsightsUI/ChaosInsightsUI.build.cs
new file mode 100644
--- /dev/null
+++ b/Engine/Plugins/ChaosInsights/Source/ChaosInsightsUI/ChaosInsightsUI.build.cs
@@ -0,0 +1,25 @@
+// Copyright Epic Games, Inc. All Rights Reserved.
+
+namespace UnrealBuildTool.Rules
+{
+	public class ChaosInsightsUI : ModuleRules
+	{
+		public ChaosInsightsUI(ReadOnlyTargetRules Target) : base(Target)
+		{
+			PrivateDependencyModuleNames.AddRange(
+			new string[]
+			{
+				"Core",
+				"CoreUObject",
+				"SlateCore",
+				"Slate",
+				"InputCore",
+				"TraceInsights",
+				"TraceInsightsCore",
+				"TraceServices",
+				"ChaosInsightsAnalysis",
+			});
+		}
+	}
+}
+
diff --git a/Engine/Plugins/ChaosInsights/Source/ChaosInsightsUI/Private/ChaosInsightsUIModule.cpp b/Engine/Plugins/ChaosInsights/Source/ChaosInsightsUI/Private/ChaosInsightsUIModule.cpp
new file mode 100644
--- /dev/null
+++ b/Engine/Plugins/ChaosInsights/Source/ChaosInsightsUI/Private/ChaosInsightsUIModule.cpp
@@ -0,0 +1,23 @@
+// Copyright Epic Games, Inc. All Rights Reserved.
+
+#include "ChaosInsightsUIModule.h"
+#include "Modules/ModuleManager.h"
+#include "Features/IModularFeatures.h"
+#include "Insights/ITimingViewExtender.h"
+
+namespace ChaosInsights 
+{
+
+	void FChaosInsightsUIModule::StartupModule()
+	{
+		IModularFeatures::Get().RegisterModularFeature(UE::Insights::Timing::TimingViewExtenderFeatureName, &TimingViewExtender);
+	}
+
+	void FChaosInsightsUIModule::ShutdownModule()
+	{
+		IModularFeatures::Get().UnregisterModularFeature(UE::Insights::Timing::TimingViewExtenderFeatureName, &TimingViewExtender);
+	}
+
+}
+
+IMPLEMENT_MODULE(ChaosInsights::FChaosInsightsUIModule, ChaosInsightsUI);
diff --git a/Engine/Plugins/ChaosInsights/Source/ChaosInsightsUI/Private/ChaosInsightsUIModule.h b/Engine/Plugins/ChaosInsights/Source/ChaosInsightsUI/Private/ChaosInsightsUIModule.h
new file mode 100644
--- /dev/null
+++ b/Engine/Plugins/ChaosInsights/Source/ChaosInsightsUI/Private/ChaosInsightsUIModule.h
@@ -0,0 +1,23 @@
+// Copyright Epic Games, Inc. All Rights Reserved.
+
+#pragma once
+
+#include "Modules/ModuleInterface.h"
+
+#include "LockRegionTrack.h"
+
+namespace ChaosInsights
+{
+
+	class FChaosInsightsUIModule : public IModuleInterface
+	{
+	public:
+		// IModuleInterface interface
+		virtual void StartupModule() override;
+		virtual void ShutdownModule() override;
+
+	private:
+		FLockRegionsSharedState TimingViewExtender;
+	};
+
+}
diff --git a/Engine/Plugins/ChaosInsights/Source/ChaosInsightsUI/Private/LockRegionTrack.cpp b/Engine/Plugins/ChaosInsights/Source/ChaosInsightsUI/Private/LockRegionTrack.cpp
new file mode 100644
--- /dev/null
+++ b/Engine/Plugins/ChaosInsights/Source/ChaosInsightsUI/Private/LockRegionTrack.cpp
@@ -0,0 +1,223 @@
+// Copyright Epic Games, Inc. All Rights Reserved.
+
+#include "LockRegionTrack.h"
+
+#include "Framework/MultiBox/MultiBoxBuilder.h"
+#include "Common/ProviderLock.h"
+#include "Insights/ITimingViewSession.h"
+#include "Insights/IUnrealInsightsModule.h"
+#include "Insights/ViewModels/TimingEventSearch.h"
+#include "Insights/ViewModels/TimingTrackViewport.h"
+#include "Insights/ViewModels/TooltipDrawState.h"
+#include "Insights/ViewModels/TimingEvent.h"
+#include "InsightsCore/Common/TimeUtils.h"
+#include "Modules/ModuleManager.h"
+
+#define LOCTEXT_NAMESPACE "LockRegionsTimingTrack"
+
+class IUnrealInsightsModule;
+class FTimingEvent;
+
+namespace ChaosInsights
+{
+	namespace Colors
+	{
+		static constexpr FColor WaitingColor{255, 15, 15};
+		static constexpr FColor AcquiredColorRead{240, 240, 100};
+		static constexpr FColor AcquiredColorWrite{15, 255, 15};
+	}
+
+	FLockRegionsSharedState::~FLockRegionsSharedState() = default;
+
+	void FLockRegionsSharedState::OnBeginSession(UE::Insights::Timing::ITimingViewSession& InSession)
+	{
+		if(!IsCurrentSession(InSession))
+		{
+			return;
+		}
+
+		LockRegionsTrack.Reset();
+	}
+
+	void FLockRegionsSharedState::OnEndSession(UE::Insights::Timing::ITimingViewSession& InSession)
+	{
+		if(!IsCurrentSession(InSession))
+		{
+			return;
+		}
+
+		LockRegionsTrack.Reset();
+	}
+
+	void FLockRegionsSharedState::Tick(UE::Insights::Timing::ITimingViewSession& InSession,
+		const TraceServices::IAnalysisSession& InAnalysisSession)
+	{
+		if(!IsCurrentSession(InSession))
+		{
+			return;
+		}
+
+		if(!LockRegionsTrack.IsValid())
+		{
+			LockRegionsTrack = MakeShared<FLockRegionsTrack>(*this);
+			LockRegionsTrack->SetOrder(FTimingTrackOrder::First);
+			LockRegionsTrack->SetVisibilityFlag(true);
+			InSession.AddScrollableTrack(LockRegionsTrack);
+		}
+	}
+
+	void FLockRegionsSharedState::ShowHideRegionsTrack()
+	{
+		bShowHideRegionsTrack = !bShowHideRegionsTrack;
+
+		if(LockRegionsTrack.IsValid())
+		{
+			LockRegionsTrack->SetVisibilityFlag(bShowHideRegionsTrack);
+		}
+
+		if(bShowHideRegionsTrack)
+		{
+			LockRegionsTrack->SetDirtyFlag();
+		}
+	}
+
+	bool FLockRegionsSharedState::IsRegionsTrackVisible() const
+	{
+		return bShowHideRegionsTrack;
+	}
+
+	bool FLockRegionsSharedState::IsCurrentSession(UE::Insights::Timing::ITimingViewSession& Session)
+	{
+		// Associate with first session seen
+		if(!TimingView)
+		{
+			TimingView = &Session;
+			return true;
+		}
+
+		return Session.GetName().IsEqual(TimingView->GetName());
+	}
+
+	INSIGHTS_IMPLEMENT_RTTI(FLockRegionsTrack);
+
+	FLockRegionsTrack::FLockRegionsTrack(FLockRegionsSharedState& InSharedState)
+		: FTimingEventsTrack(TEXT("Physics Scene Locks")), SharedState(InSharedState)
+	{
+		IUnrealInsightsModule& UnrealInsightsModule = FModuleManager::LoadModuleChecked<IUnrealInsightsModule>("TraceInsights");
+		AnalysisSession = UnrealInsightsModule.GetAnalysisSession();
+	}
+
+	FLockRegionsTrack::~FLockRegionsTrack() = default;
+
+	void FLockRegionsTrack::InitTooltip(FTooltipDrawState& InOutTooltip, const ITimingEvent& InTooltipEvent) const
+	{
+		if(InTooltipEvent.CheckTrack(this) && InTooltipEvent.Is<FTimingEvent>())
+		{
+			const FTimingEvent& TooltipEvent = InTooltipEvent.As<FTimingEvent>();
+
+			auto MatchEvent = [this, &TooltipEvent](double InStartTime, double InEndTime, uint32 InDepth)
+				{
+					return InDepth == TooltipEvent.GetDepth()
+						&& InStartTime == TooltipEvent.GetStartTime()
+						&& InEndTime == TooltipEvent.GetEndTime();
+				};
+
+			FTimingEventSearchParameters SearchParameters(TooltipEvent.GetStartTime(), TooltipEvent.GetEndTime(), ETimingEventSearchFlags::StopAtFirstMatch, MatchEvent);
+
+			auto SearchLambda = [this, &InOutTooltip, &TooltipEvent](double InFoundStartTime, double InFoundEndTime, uint32 InFoundDepth, const ChaosInsightsAnalysis::FLockRegion& InRegion)
+				{
+					InOutTooltip.Reset();
+					InOutTooltip.AddTitle(InRegion.Text, FLinearColor::White);
+					InOutTooltip.AddNameValueTextLine(TEXT("Type:"), InRegion.bIsWrite ? TEXT("Write") : TEXT("Read"));
+					InOutTooltip.AddNameValueTextLine(TEXT("Wait Duration:"), UE::Insights::FormatTimeAuto(InRegion.AcquireTime - InRegion.BeginTime));
+					InOutTooltip.AddNameValueTextLine(TEXT("Exec Duration:"), UE::Insights::FormatTimeAuto(InRegion.EndTime - InRegion.AcquireTime));
+					InOutTooltip.AddNameValueTextLine(TEXT("Max Lock Depth:"), FString::FromInt(InRegion.LockCount));
+					InOutTooltip.UpdateLayout();
+				};
+
+			FindRegionEvent(SearchParameters, SearchLambda);
+		}
+	}
+
+	void FLockRegionsTrack::BuildDrawState(ITimingEventsTrackDrawStateBuilder& Builder,
+		const ITimingTrackUpdateContext& Context)
+	{
+		const FTimingTrackViewport& Viewport = Context.GetViewport();
+
+		const ChaosInsightsAnalysis::ILockRegionProvider& RegionProvider = ChaosInsightsAnalysis::ReadRegionProvider(*AnalysisSession);
+		TraceServices::FProviderReadScopeLock RegionProviderScopedLock(RegionProvider);
+
+		int32 CurDepth = 0;
+		RegionProvider.ForEachLane([this, Viewport, &CurDepth, &Builder](const ChaosInsightsAnalysis::FLockRegionLane& Lane, const int32 Depth)
+			{
+				bool RegionHadEvents = false;
+				Lane.ForEachRegionInRange(Viewport.GetStartTime(), Viewport.GetEndTime(), [&Builder, &RegionHadEvents, &CurDepth](const ChaosInsightsAnalysis::FLockRegion& Region) -> bool
+					{
+						RegionHadEvents = true;
+
+						const uint32 WaitColorBytes = Colors::WaitingColor.DWColor();
+						const uint32 AcquiredColorBytes = Region.bIsWrite ? Colors::AcquiredColorWrite.DWColor() : Colors::AcquiredColorRead.DWColor();
+
+						Builder.AddEvent(Region.BeginTime, Region.AcquireTime, CurDepth, Region.Text, 0, WaitColorBytes);
+						Builder.AddEvent(Region.AcquireTime, Region.EndTime, CurDepth, Region.Text, 0, AcquiredColorBytes);
+
+						return true;
+					});
+
+				if(RegionHadEvents)
+				{
+					CurDepth++;
+				}
+			});
+	}
+
+	const TSharedPtr<const ITimingEvent> FLockRegionsTrack::SearchEvent(const FTimingEventSearchParameters& InSearchParameters) const
+	{
+		TSharedPtr<const ITimingEvent> FoundEvent;
+
+		auto SearchLambda = [this, &FoundEvent](double InFoundStartTime, double InFoundEndTime, uint32 InFoundDepth, const ChaosInsightsAnalysis::FLockRegion& InEvent)
+			{
+				FoundEvent = MakeShared<FTimingEvent>(SharedThis(this), InFoundStartTime, InFoundEndTime, InFoundDepth, reinterpret_cast<uint64>(InEvent.Text));
+			};
+
+		FindRegionEvent(InSearchParameters, SearchLambda);
+
+		return FoundEvent;
+	}
+
+	bool FLockRegionsTrack::FindRegionEvent(const FTimingEventSearchParameters& InParameters, TFunctionRef<void(double, double, uint32, const ChaosInsightsAnalysis::FLockRegion&)> InFoundPredicate) const
+	{
+		{
+			// If the query start time is larger than the end of the session return false.
+			TraceServices::FAnalysisSessionReadScope SessionReadScope(*AnalysisSession.Get());
+			if(AnalysisSession.IsValid() && InParameters.StartTime > AnalysisSession->GetDurationSeconds())
+			{
+				return false;
+			}
+		}
+
+		auto SearchLambda = [this](TTimingEventSearch<ChaosInsightsAnalysis::FLockRegion>::FContext& InContext)
+			{
+				const ChaosInsightsAnalysis::ILockRegionProvider& RegionProvider = ChaosInsightsAnalysis::ReadRegionProvider(*AnalysisSession);
+				TraceServices::FProviderReadScopeLock RegionProviderScopedLock(RegionProvider);
+
+				auto RegionSearchLambda = [&InContext](const ChaosInsightsAnalysis::FLockRegion& Region)
+					{
+						InContext.Check(Region.BeginTime, Region.EndTime, Region.Depth, Region);
+
+						if(!InContext.ShouldContinueSearching())
+						{
+							return false;
+						}
+
+						return true;
+					};
+
+				RegionProvider.ForEachRegionInRange(InContext.GetParameters().StartTime, InContext.GetParameters().EndTime, RegionSearchLambda);
+			};
+
+		return TTimingEventSearch<ChaosInsightsAnalysis::FLockRegion>::Search(InParameters, SearchLambda, TTimingEventSearch<ChaosInsightsAnalysis::FLockRegion>::NoFilter, InFoundPredicate, TTimingEventSearch<ChaosInsightsAnalysis::FLockRegion>::NoMatch);
+	}
+}
+
+#undef LOCTEXT_NAMESPACE
diff --git a/Engine/Plugins/ChaosInsights/Source/ChaosInsightsUI/Private/LockRegionTrack.h b/Engine/Plugins/ChaosInsights/Source/ChaosInsightsUI/Private/LockRegionTrack.h
new file mode 100644
--- /dev/null
+++ b/Engine/Plugins/ChaosInsights/Source/ChaosInsightsUI/Private/LockRegionTrack.h
@@ -0,0 +1,57 @@
+// Copyright Epic Games, Inc. All Rights Reserved.
+
+#pragma once
+
+#include "Insights/ITimingViewExtender.h"
+#include "Insights/ViewModels/TimingEventsTrack.h"
+#include "ChaosInsightsAnalysis/Model/LockRegions.h"
+
+namespace ChaosInsights
+{
+
+	class FLockRegionsTrack;
+
+	class FLockRegionsSharedState : public UE::Insights::Timing::ITimingViewExtender, public TSharedFromThis<FLockRegionsSharedState>
+	{
+		friend class FLockRegionsTrack;
+
+	public:
+		virtual ~FLockRegionsSharedState() override;
+
+		// ITimingViewExtender interface
+		virtual void OnBeginSession(UE::Insights::Timing::ITimingViewSession& InSession) override;
+		virtual void OnEndSession(UE::Insights::Timing::ITimingViewSession& InSession) override;
+		virtual void Tick(UE::Insights::Timing::ITimingViewSession& InSession, const TraceServices::IAnalysisSession& InAnalysisSession) override;
+
+		void ShowHideRegionsTrack();
+		bool IsRegionsTrackVisible() const;
+
+	private:
+		bool IsCurrentSession(UE::Insights::Timing::ITimingViewSession& Session);
+
+		UE::Insights::Timing::ITimingViewSession* TimingView = nullptr;
+		TSharedPtr<FLockRegionsTrack> LockRegionsTrack;
+		bool bShowHideRegionsTrack = true;
+	};
+
+	class FLockRegionsTrack : public FTimingEventsTrack
+	{
+		INSIGHTS_DECLARE_RTTI(FLockRegionsTrack, FTimingEventsTrack)
+
+	public:
+		FLockRegionsTrack(FLockRegionsSharedState& InSharedState);
+		virtual ~FLockRegionsTrack() override;
+
+		virtual void InitTooltip(FTooltipDrawState& InOutTooltip, const ITimingEvent& InTooltipEvent) const override;
+		virtual const TSharedPtr<const ITimingEvent> SearchEvent(const FTimingEventSearchParameters& InSearchParameters) const override;
+		virtual void BuildDrawState(ITimingEventsTrackDrawStateBuilder& Builder, const ITimingTrackUpdateContext& Context) override;
+
+	protected:
+		bool FindRegionEvent(const FTimingEventSearchParameters& InParameters, TFunctionRef<void(double, double, uint32, const ChaosInsightsAnalysis::FLockRegion&)> InFoundPredicate) const;
+
+	private:
+		FLockRegionsSharedState& SharedState;
+		TSharedPtr<const TraceServices::IAnalysisSession> AnalysisSession;
+	};
+
+}
diff --git a/Engine/Programs/UnrealInsights/Config/DefaultEngine.ini b/Engine/Programs/UnrealInsights/Config/DefaultEngine.ini
--- a/Engine/Programs/UnrealInsights/Config/DefaultEngine.ini
+++ b/Engine/Programs/UnrealInsights/Config/DefaultEngine.ini
@@ -10,6 +10,7 @@
 +ProgramEnabledPlugins=VisualStudioSourceCodeAccess
 +ProgramEnabledPlugins=XcodeSourceCodeAccess
 +ProgramEnabledPlugins=MassInsights
++ProgramEnabledPlugins=ChaosInsights
 
 [/Script/MacTargetPlatform.XcodeProjectSettings]
 TemplateMacPlist=(FilePath="/Engine/Source/Programs/UnrealInsights/Resources/Mac/Info.plist")
diff --git a/Engine/Source/Runtime/Experimental/Chaos/Private/ChaosInsights/ChaosInsights.cpp b/Engine/Source/Runtime/Experimental/Chaos/Private/ChaosInsights/ChaosInsights.cpp
new file mode 100644
--- /dev/null
+++ b/Engine/Source/Runtime/Experimental/Chaos/Private/ChaosInsights/ChaosInsights.cpp
@@ -0,0 +1,53 @@
+// Copyright Epic Games, Inc. All Rights Reserved.
+
+#if UE_TRACE_ENABLED
+
+#include "ChaosInsights/ChaosInsightsMacros.h"
+
+#include "CoreMinimal.h"
+#include "Trace/Trace.h"
+#include "Trace/Trace.inl"
+#include "Misc/CString.h"
+#include "HAL/PlatformTime.h"
+
+UE_TRACE_CHANNEL(ChaosLocksChannel)
+
+UE_TRACE_EVENT_BEGIN(Chaos, LockAcquireBegin)
+	UE_TRACE_EVENT_FIELD(uint64, Cycle)
+	UE_TRACE_EVENT_FIELD(bool, bIsWrite)
+UE_TRACE_EVENT_END()
+
+UE_TRACE_EVENT_BEGIN(Chaos, LockAcquired)
+	UE_TRACE_EVENT_FIELD(uint64, Cycle)
+UE_TRACE_EVENT_END()
+
+UE_TRACE_EVENT_BEGIN(Chaos, LockAcquireEnd)
+	UE_TRACE_EVENT_FIELD(uint64, Cycle)
+UE_TRACE_EVENT_END()
+
+namespace Chaos::Insights
+{
+	uint64 BeginLockAcquireEvent(ELockEventType Type)
+	{
+		uint64 Id = FPlatformTime::Cycles64();
+		UE_TRACE_LOG(Chaos, LockAcquireBegin, ChaosLocksChannel)
+			<< LockAcquireBegin.Cycle(FPlatformTime::Cycles64())
+			<< LockAcquireBegin.bIsWrite(Type == ELockEventType::RWLockWriteLock || Type == ELockEventType::Mutex);
+
+		return Id;
+	}
+
+	void AcquiredLock()
+	{
+		UE_TRACE_LOG(Chaos, LockAcquired, ChaosLocksChannel)
+			<< LockAcquired.Cycle(FPlatformTime::Cycles64());
+	}
+
+	void EndLockAcquireEvent()
+	{
+		UE_TRACE_LOG(Chaos, LockAcquireEnd, ChaosLocksChannel)
+			<< LockAcquireEnd.Cycle(FPlatformTime::Cycles64());
+	}
+}
+
+#endif // UE_TRACE_ENABLED
diff --git a/Engine/Source/Runtime/Experimental/Chaos/Public/ChaosInsights/ChaosInsightsMacros.h b/Engine/Source/Runtime/Experimental/Chaos/Public/ChaosInsights/ChaosInsightsMacros.h
new file mode 100644
--- /dev/null
+++ b/Engine/Source/Runtime/Experimental/Chaos/Public/ChaosInsights/ChaosInsightsMacros.h
@@ -0,0 +1,33 @@
+// Copyright Epic Games, Inc. All Rights Reserved.
+
+#pragma once
+
+#if UE_TRACE_ENABLED
+
+#include "CoreTypes.h"
+
+namespace Chaos::Insights
+{
+	enum class ELockEventType
+	{
+		Mutex,
+		RWLockReadLock,
+		RWLockWriteLock
+	};
+
+	CHAOS_API uint64 BeginLockAcquireEvent(ELockEventType Type);
+	CHAOS_API void AcquiredLock();
+	CHAOS_API void EndLockAcquireEvent();
+}
+
+#define TRACE_CHAOS_BEGIN_LOCK(Type)	Chaos::Insights::BeginLockAcquireEvent(Type);
+#define TRACE_CHAOS_ACQUIRE_LOCK()		Chaos::Insights::AcquiredLock();
+#define TRACE_CHAOS_END_LOCK()			Chaos::Insights::EndLockAcquireEvent();
+
+#else
+
+#define TRACE_CHAOS_BEGIN_LOCK(...)
+#define TRACE_CHAOS_ACQUIRE_LOCK(...)
+#define TRACE_CHAOS_END_LOCK(...)
+
+#endif // UE_TRACE_ENABLED
\ No newline at end of file
diff --git a/Engine/Source/Runtime/Experimental/Chaos/Public/Framework/Threading.h b/Engine/Source/Runtime/Experimental/Chaos/Public/Framework/Threading.h
--- a/Engine/Source/Runtime/Experimental/Chaos/Public/Framework/Threading.h
+++ b/Engine/Source/Runtime/Experimental/Chaos/Public/Framework/Threading.h
@@ -13,6 +13,7 @@
 #include "HAL/CriticalSection.h"
 #include "AutoRTFM.h"
 #include "Templates/SharedPointer.h"
+#include "ChaosInsights/ChaosInsightsMacros.h"
 
 #ifndef PHYSICS_THREAD_CONTEXT
 	#if (!UE_BUILD_SHIPPING && !UE_BUILD_TEST)
@@ -541,6 +542,7 @@ FORCEINLINE void EnsureIsInGameThreadContext()
 
 		void ReadLock()
 		{
+			TRACE_CHAOS_BEGIN_LOCK(Chaos::Insights::ELockEventType::RWLockReadLock);
 			if(ThreadingPrivate::GetThreadReadDepth(this) == 0)
 			{
 				TMutexScopeLock<MutexType> Guard(Mutex);
@@ -563,10 +565,14 @@ FORCEINLINE void EnsureIsInGameThreadContext()
 #endif
 
 			CHAOS_RECORD_ENTER_READ_LOCK;
+
+			TRACE_CHAOS_ACQUIRE_LOCK();
 		}
 
 		void WriteLock()
 		{
+			TRACE_CHAOS_BEGIN_LOCK(Chaos::Insights::ELockEventType::RWLockWriteLock);
+
 #if CHAOS_SCENE_LOCK_CHECKS
 			if(ThreadingPrivate::GetThreadReadDepth(this) > 0)
 			{
@@ -605,6 +611,7 @@ FORCEINLINE void EnsureIsInGameThreadContext()
 			// Write lock means we can access game thread data, so set the right context
 			FPhysicsThreadContext::Get().IncGameThreadContext();
 #endif
+			TRACE_CHAOS_ACQUIRE_LOCK();
 		}
 
 		void ReadUnlock()
@@ -618,6 +625,8 @@ FORCEINLINE void EnsureIsInGameThreadContext()
 			// Read lock is released, the gamethread context is gone
 			FPhysicsThreadContext::Get().DecGameThreadContext();
 #endif
+		
+			TRACE_CHAOS_END_LOCK();
 		}
 
 		void WriteUnlock()
@@ -629,6 +638,8 @@ FORCEINLINE void EnsureIsInGameThreadContext()
 			// Write lock is released, the gamethread context is gone
 			FPhysicsThreadContext::Get().DecGameThreadContext();
 #endif
+
+			TRACE_CHAOS_END_LOCK();
 		}
 
 	private:
@@ -745,6 +756,8 @@ FORCEINLINE void EnsureIsInGameThreadContext()
 
 		void ReadLock()
 		{
+			TRACE_CHAOS_BEGIN_LOCK(Chaos::Insights::ELockEventType::RWLockReadLock)
+
 			FRwLockInfo ThreadInfo(FPlatformTLS::GetTlsValue(TlsSlot));
 			ThreadInfo.ThreadReadDepth++;
 			FPlatformTLS::SetTlsValue(TlsSlot, ThreadInfo.GetTlsSlotValue());
@@ -758,10 +771,12 @@ FORCEINLINE void EnsureIsInGameThreadContext()
 			// Read lock means we can access game thread data, so set the right context
 			FPhysicsThreadContext::Get().IncGameThreadContext();
 #endif
-
+			TRACE_CHAOS_ACQUIRE_LOCK()
 		}
 		void WriteLock()
 		{
+			TRACE_CHAOS_BEGIN_LOCK(Chaos::Insights::ELockEventType::RWLockWriteLock)
+
 			FRwLockInfo ThreadInfo(FPlatformTLS::GetTlsValue(TlsSlot));
 			ThreadInfo.ThreadWriteDepth++;
 			FPlatformTLS::SetTlsValue(TlsSlot, ThreadInfo.GetTlsSlotValue());
@@ -780,6 +795,7 @@ FORCEINLINE void EnsureIsInGameThreadContext()
 			// Write lock means we can access game thread data, so set the right context
 			FPhysicsThreadContext::Get().IncGameThreadContext();
 #endif
+			TRACE_CHAOS_ACQUIRE_LOCK()
 		}
 		void ReadUnlock()
 		{
@@ -796,6 +812,7 @@ FORCEINLINE void EnsureIsInGameThreadContext()
 			// Read lock is released, the gamethread context is gone
 			FPhysicsThreadContext::Get().DecGameThreadContext();
 #endif
+			TRACE_CHAOS_END_LOCK()
 		}
 		void WriteUnlock()
 		{
@@ -811,6 +828,7 @@ FORCEINLINE void EnsureIsInGameThreadContext()
 			// Write lock is released, the gamethread context is gone
 			FPhysicsThreadContext::Get().DecGameThreadContext();
 #endif
+			TRACE_CHAOS_END_LOCK()
 		}
 
 	private:
