Fix Unreal Insights Trace crashes after enabling and disabling the Timing Tab.

diff --git a/Engine/Plugins/ChaosInsights/Source/ChaosInsightsUI/Private/LockRegionTrack.cpp b/Engine/Plugins/ChaosInsights/Source/ChaosInsightsUI/Private/LockRegionTrack.cpp
--- a/Engine/Plugins/ChaosInsights/Source/ChaosInsightsUI/Private/LockRegionTrack.cpp
+++ b/Engine/Plugins/ChaosInsights/Source/ChaosInsightsUI/Private/LockRegionTrack.cpp
@@ -95,7 +95,7 @@ namespace ChaosInsights
 			return true;
 		}
 
-		return Session.GetName().IsEqual(TimingView->GetName());
+		return &Session == TimingView;
 	}
 
 	INSIGHTS_IMPLEMENT_RTTI(FLockRegionsTrack);
