[Insights] Chaos Insights: Fixed crash due to usage of a ITimingViewSession pointer after the TimingView widget was destroyed.

diff --git a/Engine/Plugins/ChaosInsights/Source/ChaosInsightsUI/Private/LockRegionTrack.cpp b/Engine/Plugins/ChaosInsights/Source/ChaosInsightsUI/Private/LockRegionTrack.cpp
--- a/Engine/Plugins/ChaosInsights/Source/ChaosInsightsUI/Private/LockRegionTrack.cpp
+++ b/Engine/Plugins/ChaosInsights/Source/ChaosInsightsUI/Private/LockRegionTrack.cpp
@@ -31,6 +31,11 @@ namespace ChaosInsights
 
 	void FLockRegionsSharedState::OnBeginSession(UE::Insights::Timing::ITimingViewSession& InSession)
 	{
+		if (TimingView == nullptr)
+		{
+			TimingView = &InSession;
+		}
+
 		if(!IsCurrentSession(InSession))
 		{
 			return;
@@ -46,6 +51,8 @@ namespace ChaosInsights
 			return;
 		}
 
+		TimingView = nullptr;
+
 		LockRegionsTrack.Reset();
 	}
 
@@ -88,13 +95,6 @@ namespace ChaosInsights
 
 	bool FLockRegionsSharedState::IsCurrentSession(UE::Insights::Timing::ITimingViewSession& Session)
 	{
-		// Associate with first session seen
-		if(!TimingView)
-		{
-			TimingView = &Session;
-			return true;
-		}
-
 		return &Session == TimingView;
 	}
 
